Unshelved from pending changelist '8925094':

- Resampled animation time track to target(default) fps. GLTF doesn't store animation framerate(this is not a problem as animation times are in seconds). Now keyframes are correctly mapped.
- Handled Cubic interpolation

diff --git a/Engine/Plugins/Editor/GLTFImporter/Source/GLTFImporter/Private/GLTF/ConversionUtilities.h b/Engine/Plugins/Editor/GLTFImporter/Source/GLTFImporter/Private/GLTF/ConversionUtilities.h
--- a/Engine/Plugins/Editor/GLTFImporter/Source/GLTFImporter/Private/GLTF/ConversionUtilities.h
+++ b/Engine/Plugins/Editor/GLTFImporter/Source/GLTFImporter/Private/GLTF/ConversionUtilities.h
@@ -34,7 +34,11 @@ namespace GLTF
 	{
 		// glTF uses a right-handed coordinate system, with Y up.
 		// UE4 uses a left-handed coordinate system, with Z up.
-		// thus apply sin(x/2) * (X,Y,Z) and also reverse the angle
+		// Quat = (qX, qY, qZ, qW) = (sin(angle/2) * aX, sin(angle/2) * aY, sin(angle/2) * aZ, cons(angle/2))
+		// where (aX, aY, aZ) - rotation axis, angle - rotation angle
+		// Y swapped with Z between these coordinate systems
+		// also, as handedness is changed rotation is inversed - hence negation
+		// therefore QuatUE = (-qX, -qZ, -qY, qw)
 
 		FQuat Result(-Quat.X, -Quat.Z, -Quat.Y, Quat.W);
 		check(Result.IsNormalized());
diff --git a/Engine/Plugins/Editor/GLTFImporter/Source/GLTFImporter/Private/GLTF/GLTFAccessor.cpp b/Engine/Plugins/Editor/GLTFImporter/Source/GLTFImporter/Private/GLTF/GLTFAccessor.cpp
--- a/Engine/Plugins/Editor/GLTFImporter/Source/GLTFImporter/Private/GLTF/GLTFAccessor.cpp
+++ b/Engine/Plugins/Editor/GLTFImporter/Source/GLTFImporter/Private/GLTF/GLTFAccessor.cpp
@@ -184,8 +184,8 @@ namespace GLTF
 		for (uint32 Index = 0; Index < Count; ++Index)
 		{
 			FVector4& Value = Buffer[Index];
-			FQuat     Quat(Value[0], Value[1], Value[2], Value[3]);
-			//	Quat    = GLTF::ConvertQuat(Quat);
+			FQuat Quat(Value[0], Value[1], Value[2], Value[3]);
+			Quat = GLTF::ConvertQuat(Quat);
 			Value.X = Quat.X;
 			Value.Y = Quat.Y;
 			Value.Z = Quat.Z;
diff --git a/Engine/Plugins/Editor/GLTFImporter/Source/GLTFImporter/Private/GLTF/GLTFReader.cpp b/Engine/Plugins/Editor/GLTFImporter/Source/GLTFImporter/Private/GLTF/GLTFReader.cpp
--- a/Engine/Plugins/Editor/GLTFImporter/Source/GLTFImporter/Private/GLTF/GLTFReader.cpp
+++ b/Engine/Plugins/Editor/GLTFImporter/Source/GLTFImporter/Private/GLTF/GLTFReader.cpp
@@ -377,8 +377,20 @@ namespace GLTF
 				if (ensure((Input != INDEX_NONE) && (Output != INDEX_NONE)))
 				{
 					FAnimation::FSampler Sampler(Asset->Accessors[Input], Asset->Accessors[Output]);
-					Sampler.Interpolation =
-						(FAnimation::EInterpolation)GetUnsignedInt(SamplerObject, TEXT("interpolation"), (uint32)Sampler.Interpolation);
+					FString InterpolationStr = GetString(SamplerObject, TEXT("interpolation"), TEXT("LINEAR"));
+					if (InterpolationStr == TEXT("LINEAR"))
+					{
+						Sampler.Interpolation = FAnimation::EInterpolation::Linear;
+					}
+					else if (InterpolationStr == TEXT("STEP"))
+					{
+						Sampler.Interpolation = FAnimation::EInterpolation::Step;
+					}
+					else if (InterpolationStr == TEXT("CUBICSPLINE"))
+					{
+						Sampler.Interpolation = FAnimation::EInterpolation::CubicSpline;
+					}
+
 					Animation.Samplers.Add(Sampler);
 				}
 			}
diff --git a/Engine/Plugins/Enterprise/DatasmithGLTFImporter/Source/DatasmithGLTFTranslator/Private/DatasmithGLTFAnimationImporter.cpp b/Engine/Plugins/Enterprise/DatasmithGLTFImporter/Source/DatasmithGLTFTranslator/Private/DatasmithGLTFAnimationImporter.cpp
--- a/Engine/Plugins/Enterprise/DatasmithGLTFImporter/Source/DatasmithGLTFTranslator/Private/DatasmithGLTFAnimationImporter.cpp
+++ b/Engine/Plugins/Enterprise/DatasmithGLTFImporter/Source/DatasmithGLTFTranslator/Private/DatasmithGLTFAnimationImporter.cpp
@@ -8,6 +8,8 @@
 #include "IDatasmithSceneElements.h"
 #include "DatasmithUtils.h"
 
+#include "Math/InterpCurve.h"
+
 namespace DatasmithGLTFImporterImpl
 {
 	EDatasmithTransformType ConvertToTransformType(GLTF::FAnimation::EPath Path)
@@ -34,6 +36,125 @@ namespace DatasmithGLTFImporterImpl
 		return FDatasmithTransformFrameInfo(FrameNumber, Vec);
 	}
 
+	template<typename TSourceType, typename TConvertToUE>
+	void ResampleTrack(
+		const FFrameRate& FrameRate,
+		const TArray<float>& FrameTimeBuffer,
+		const TSourceType* FrameSourceBuffer,
+		GLTF::FAnimation::EInterpolation Interpolation,
+		TConvertToUE ConvertToUE,
+		IDatasmithTransformAnimationElement& AnimationElement,
+		const EDatasmithTransformType TransformType)
+	{
+		FInterpCurve<TSourceType> InterpCurve;
+		float MinKey = FLT_MAX;
+		float MaxKey = -FLT_MAX;
+
+		for (int FrameIndex = 0; FrameIndex < FrameTimeBuffer.Num(); ++FrameIndex)
+		{
+			float FrameTime = FrameTimeBuffer[FrameIndex];
+
+			FInterpCurvePoint<TSourceType> CurvePoint;
+
+			if (Interpolation == GLTF::FAnimation::EInterpolation::CubicSpline)
+			{
+				int32 BufferIndex = FrameIndex * 3 + 1; // spline vertex value located between two tangents
+
+				// "When used with CUBICSPLINE interpolation, tangents (ak, bk) and values (vk) are grouped within keyframes:
+				// a1, a2, ...an, v1, v2, ...vn, b1, b2, ...bn"
+				// [REDACTED_PROFILE_URL]/glTF/tree/master/specification/2.0#animations
+				CurvePoint = FInterpCurvePoint<TSourceType>(FrameTime,
+					FrameSourceBuffer[FrameIndex * 3 + 1], // value
+					FrameSourceBuffer[FrameIndex * 3 + 0], // in-tangent
+					FrameSourceBuffer[FrameIndex * 3 + 2], // out-tangent
+					EInterpCurveMode::CIM_CurveUser
+					);
+			}
+			else
+			{
+				int32 BufferIndex = FrameIndex;
+				CurvePoint = FInterpCurvePoint<TSourceType>(FrameTime, FrameSourceBuffer[BufferIndex]);
+
+				switch (Interpolation)
+				{
+				case GLTF::FAnimation::EInterpolation::Step:
+					CurvePoint.InterpMode = EInterpCurveMode::CIM_Constant;
+					break;
+				case GLTF::FAnimation::EInterpolation::Linear:
+					CurvePoint.InterpMode = EInterpCurveMode::CIM_Linear;
+					break;
+				};
+			}
+
+			InterpCurve.Points.Add(CurvePoint);
+
+			MinKey = FMath::Min(MinKey, FrameTime);
+			MaxKey = FMath::Max(MaxKey, FrameTime);
+		}
+
+		FFrameNumber StartFrame = FrameRate.AsFrameNumber(MinKey);
+
+		// If we use AsFrameNumber it will floor, and we might lose the very end of the animation
+		const double TimeAsFrame = (double(MaxKey) * FrameRate.Numerator) / FrameRate.Denominator;
+		FFrameNumber EndFrame = FFrameNumber(static_cast<int32>(FMath::CeilToDouble(TimeAsFrame)));
+
+		// We go to EndFrame.Value+1 here so that if its a 2 second animation at 30fps, frame 60 belongs
+		// to the actual animation, as opposed to being range [0, 59]. This guarantees that the animation will
+		// actually complete within its range, which is necessary in order to play it correctly at runtime
+		for (int32 Frame = StartFrame.Value; Frame <= EndFrame.Value + 1; ++Frame)
+		{
+			float TimeSeconds = FrameRate.AsSeconds(Frame);
+
+			AnimationElement.AddFrame(TransformType, FDatasmithTransformFrameInfo(Frame, ConvertToUE(InterpCurve.Eval(TimeSeconds))));
+		}
+	}
+
+	class FAnimationResampler
+	{
+	public:
+
+		FAnimationResampler(FFrameRate InFrameRate, const TArray<float>& InFrameTimeBuffer, IDatasmithTransformAnimationElement& InAnimationElement) 
+			: FrameRate(InFrameRate)
+			, FrameTimeBuffer(InFrameTimeBuffer)
+			, AnimationElement(InAnimationElement)
+		{}
+
+		void ResampleRotation(
+			GLTF::FAnimation::EInterpolation Interpolation,
+			const FQuat* FrameSourceBuffer)
+		{
+			ResampleTrack<FQuat>(FrameRate, 
+				FrameTimeBuffer, FrameSourceBuffer, Interpolation,
+				[](const FQuat& Quat) { return Quat.Euler(); },
+				AnimationElement, EDatasmithTransformType::Rotation);
+		}
+
+		void ResampleTranslation(
+			GLTF::FAnimation::EInterpolation Interpolation,
+			const FVector* FrameSourceBuffer,
+			float ScaleFactor)
+		{
+			ResampleTrack<FVector>(FrameRate,
+				FrameTimeBuffer, FrameSourceBuffer, Interpolation,
+				[ScaleFactor](const FVector& Vec) { return Vec * ScaleFactor; },
+				AnimationElement, EDatasmithTransformType::Translation);
+		}
+
+		void ResampleScale(
+			GLTF::FAnimation::EInterpolation Interpolation,
+			const FVector* FrameSourceBuffer)
+		{
+			ResampleTrack<FVector>(FrameRate,
+				FrameTimeBuffer, FrameSourceBuffer, Interpolation,
+				[](const FVector& Vec) { return Vec; },
+				AnimationElement, EDatasmithTransformType::Scale);
+		}
+
+	private:
+		FFrameRate FrameRate;
+		const TArray<float>& FrameTimeBuffer;
+		IDatasmithTransformAnimationElement& AnimationElement;
+	};
 }
 
 FDatasmithGLTFAnimationImporter::FDatasmithGLTFAnimationImporter(TArray<GLTF::FLogMessage>& LogMessages)
@@ -66,7 +187,7 @@ void FDatasmithGLTFAnimationImporter::CreateAnimations(const GLTF::FAsset& GLTFA
 				LogMessages.Emplace(GLTF::EMessageSeverity::Error, TEXT("Morph animations aren't supported: ") + Animation.Name);
 		}
 
-		const float FrameRate = SequenceElement->GetFrameRate();
+		FFrameRate FrameRate = FFrameRate(FMath::RoundToInt(SequenceElement->GetFrameRate()), 1);
 		for (const auto& NodeChannelPair : NodeChannelMap)
 		{
 			const GLTF::FNode*                        Node     = NodeChannelPair.Get<0>();
@@ -74,7 +195,7 @@ void FDatasmithGLTFAnimationImporter::CreateAnimations(const GLTF::FAsset& GLTFA
 
 			TSharedRef<IDatasmithTransformAnimationElement> AnimationElement = FDatasmithSceneFactory::CreateTransformAnimation(*Node->Name);
 
-			CreateAnimationFrames(Animation, Channels, FrameRate, *AnimationElement);
+			ResampleAnimationFrames(Animation, Channels, FrameRate, *AnimationElement);
 			SequenceElement->AddAnimation(AnimationElement);
 		}
 
@@ -83,10 +204,10 @@ void FDatasmithGLTFAnimationImporter::CreateAnimations(const GLTF::FAsset& GLTFA
 	}
 }
 
-uint32 FDatasmithGLTFAnimationImporter::CreateAnimationFrames(const GLTF::FAnimation&                   Animation,
-                                                              const TArray<GLTF::FAnimation::FChannel>& Channels,
-                                                              float                                     FrameRate,
-                                                              IDatasmithTransformAnimationElement&      AnimationElement)
+uint32 FDatasmithGLTFAnimationImporter::ResampleAnimationFrames(const GLTF::FAnimation& Animation,
+	const TArray<GLTF::FAnimation::FChannel>& Channels,
+	FFrameRate FrameRate,
+	IDatasmithTransformAnimationElement& AnimationElement)
 {
 	using namespace DatasmithGLTFImporterImpl;
 
@@ -101,7 +222,9 @@ uint32 FDatasmithGLTFAnimationImporter::CreateAnimationFrames(const GLTF::FAnima
 
 	EDatasmithTransformChannels ActiveChannels = EDatasmithTransformChannels::None;
 
-	bool bProcessedPath[3] = {false, false, false};
+	FAnimationResampler Resampler(FrameRate, FrameTimeBuffer, AnimationElement);
+
+	bool bProcessedPath[3] = { false, false, false };
 	for (const GLTF::FAnimation::FChannel& Channel : Channels)
 	{
 		check(bProcessedPath[(int32)Channel.Target.Path] == false);
@@ -115,67 +238,45 @@ uint32 FDatasmithGLTFAnimationImporter::CreateAnimationFrames(const GLTF::FAnima
 		int32 Index = 0;
 		switch (Channel.Target.Path)
 		{
-			case GLTF::FAnimation::EPath::Rotation:
-				// always vec4
-				FrameDataBuffer.SetNumUninitialized(Sampler.Output.Count * 4);
-				Sampler.Output.GetVec4Array(reinterpret_cast<FVector4*>(FrameDataBuffer.GetData()));
+		case GLTF::FAnimation::EPath::Rotation:
+		{
+			FrameDataBuffer.SetNumUninitialized(Sampler.Output.Count * 4);
+			// GTFL Accessor public api returns quaternions into FVector4 array
+			Sampler.Output.GetQuatArray(reinterpret_cast<FVector4*>(FrameDataBuffer.GetData()));
 
-				for (float Time : FrameTimeBuffer)
-				{
-					const float* Values = &FrameDataBuffer[(Index++) * 4];
-
-					// glTF uses a right-handed coordinate system, with Y up.
-					// UE4 uses a left-handed coordinate system, with Z up.
-					// Quat = (qX, qY, qZ, qW) = (sin(angle/2) * aX, sin(angle/2) * aY, sin(angle/2) * aZ, cons(angle/2))
-					// where (aX, aY, aZ) - rotation axis, angle - rotation angle
-					// Y swapped with Z between these coordinate systems
-					// also, as handedness is changed rotation is inversed - hence negation
-					// therefore QuatUE = (-qX, -qZ, -qY, qw)
-					const FQuat  Quat(-Values[0], -Values[2], -Values[1], Values[3]);
-
-					const FDatasmithTransformFrameInfo FrameInfo = CreateFrameInfo(Time * FrameRate, Quat.Euler());
-					AnimationElement.AddFrame(TransformType, FrameInfo);
-				}
-
-				ActiveChannels = ActiveChannels | FDatasmithAnimationUtils::SetChannelTypeComponents(ETransformChannelComponents::All, TransformType);
-				break;
-			case GLTF::FAnimation::EPath::Translation:
-			case GLTF::FAnimation::EPath::Scale:
-				// always vec3
-				FrameDataBuffer.SetNumUninitialized(Sampler.Output.Count * 3);
-				Sampler.Output.GetCoordArray(reinterpret_cast<FVector*>(FrameDataBuffer.GetData()));
+			Resampler.ResampleRotation(Sampler.Interpolation, reinterpret_cast<FQuat*>(FrameDataBuffer.GetData()));
 
-				if (Channel.Target.Path == GLTF::FAnimation::EPath::Translation)
-				{
-					for (float Time : FrameTimeBuffer)
-					{
-						const float*  Values = &FrameDataBuffer[(Index++) * 3];
-						const FVector Vec(Values[0], Values[1], Values[2]);
-
-						const FDatasmithTransformFrameInfo FrameInfo = CreateFrameInfo(Time * FrameRate, Vec * ScaleFactor);
-						AnimationElement.AddFrame(TransformType, FrameInfo);
-					}
-				}
-				else
-				{
-					for (float Time : FrameTimeBuffer)
-					{
-						const float*  Values = &FrameDataBuffer[(Index++) * 3];
-						const FVector Vec(Values[0], Values[1], Values[2]);
-
-						const FDatasmithTransformFrameInfo FrameInfo = CreateFrameInfo(Time * FrameRate, Vec);
-						AnimationElement.AddFrame(TransformType, FrameInfo);
-					}
-				}
-
-				ActiveChannels = ActiveChannels | FDatasmithAnimationUtils::SetChannelTypeComponents(ETransformChannelComponents::All, TransformType);
-				break;
-			default:
-				check(false);
-				break;
+			ActiveChannels = ActiveChannels | FDatasmithAnimationUtils::SetChannelTypeComponents(ETransformChannelComponents::All, TransformType);
+			break;
+		}
+		case GLTF::FAnimation::EPath::Translation:
+		{
+			FrameDataBuffer.SetNumUninitialized(Sampler.Output.Count * 3);
+			FVector* SourceData = reinterpret_cast<FVector*>(FrameDataBuffer.GetData());
+			Sampler.Output.GetCoordArray(SourceData);
+
+			Resampler.ResampleTranslation(Sampler.Interpolation, SourceData, ScaleFactor);
+
+			ActiveChannels = ActiveChannels | FDatasmithAnimationUtils::SetChannelTypeComponents(ETransformChannelComponents::All, TransformType);
+			break;
+		}
+		case GLTF::FAnimation::EPath::Scale:
+		{
+			FrameDataBuffer.SetNumUninitialized(Sampler.Output.Count * 3);
+			FVector* SourceData = reinterpret_cast<FVector*>(FrameDataBuffer.GetData());
+			Sampler.Output.GetCoordArray(SourceData);
+
+			Resampler.ResampleScale(Sampler.Interpolation, SourceData);
+
+			ActiveChannels = ActiveChannels | FDatasmithAnimationUtils::SetChannelTypeComponents(ETransformChannelComponents::All, TransformType);
+			break;
+		}
+		default:
+			ensure(false);
+			break;
 		}
 
-		FrameCount                                 = FMath::Max(Sampler.Input.Count, FrameCount);
+		FrameCount = FMath::Max(Sampler.Input.Count, FrameCount);
 		bProcessedPath[(int32)Channel.Target.Path] = true;
 	}
 
diff --git a/Engine/Plugins/Enterprise/DatasmithGLTFImporter/Source/DatasmithGLTFTranslator/Private/DatasmithGLTFAnimationImporter.h b/Engine/Plugins/Enterprise/DatasmithGLTFImporter/Source/DatasmithGLTFTranslator/Private/DatasmithGLTFAnimationImporter.h
--- a/Engine/Plugins/Enterprise/DatasmithGLTFImporter/Source/DatasmithGLTFTranslator/Private/DatasmithGLTFAnimationImporter.h
+++ b/Engine/Plugins/Enterprise/DatasmithGLTFImporter/Source/DatasmithGLTFTranslator/Private/DatasmithGLTFAnimationImporter.h
@@ -5,6 +5,8 @@
 #include "GLTFAnimation.h"
 #include "GLTFLogger.h"
 
+#include "Misc/FrameRate.h"
+
 class IDatasmithScene;
 class IDatasmithTransformAnimationElement;
 class IDatasmithLevelSequenceElement;
@@ -28,10 +30,10 @@ public:
 	void SetUniformScale(float Scale);
 
 private:
-	uint32 CreateAnimationFrames(const GLTF::FAnimation&                   Animation,
-	                             const TArray<GLTF::FAnimation::FChannel>& Channels,
-	                             float                                     FrameRate,
-	                             IDatasmithTransformAnimationElement&      AnimationElement);
+	uint32 ResampleAnimationFrames(const GLTF::FAnimation& Animation,
+		const TArray<GLTF::FAnimation::FChannel>& Channels,
+		FFrameRate FrameRate,
+		IDatasmithTransformAnimationElement& AnimationElement);
 
 private:
 	using FNodeChannelMap = TMap<const GLTF::FNode*, TArray<GLTF::FAnimation::FChannel> >;
