Extensions not supported message are displayed while that extensions are implemented

diff --git a/Engine/Plugins/Enterprise/DatasmithGLTFImporter/Source/DatasmithGLTFTranslator/Private/DatasmithGLTFImporter.cpp b/Engine/Plugins/Enterprise/DatasmithGLTFImporter/Source/DatasmithGLTFTranslator/Private/DatasmithGLTFImporter.cpp
--- a/Engine/Plugins/Enterprise/DatasmithGLTFImporter/Source/DatasmithGLTFTranslator/Private/DatasmithGLTFImporter.cpp
+++ b/Engine/Plugins/Enterprise/DatasmithGLTFImporter/Source/DatasmithGLTFTranslator/Private/DatasmithGLTFImporter.cpp
@@ -98,7 +98,12 @@ bool FDatasmithGLTFImporter::OpenFile(const FString& InFileName)
 
 	// check extensions supported
 	static const TArray<GLTF::EExtension> SupportedExtensions = {GLTF::EExtension::KHR_MaterialsPbrSpecularGlossiness,
-	                                                             GLTF::EExtension::KHR_MaterialsUnlit, GLTF::EExtension::KHR_LightsPunctual};
+																 GLTF::EExtension::KHR_MaterialsUnlit,
+																 GLTF::EExtension::KHR_LightsPunctual,
+																 GLTF::EExtension::KHR_MaterialsClearCoat,
+																 GLTF::EExtension::KHR_MaterialsTransmission,
+																 GLTF::EExtension::KHR_MaterialsSheen};
+
 	for (GLTF::EExtension Extension : GLTFAsset->ExtensionsUsed)
 	{
 		if (SupportedExtensions.Find(Extension) == INDEX_NONE)
