Support material variants

diff --git a/Engine/Plugins/Editor/GLTFImporter/Source/GLTFCore/Private/GLTF/ExtensionsHandler.cpp b/Engine/Plugins/Editor/GLTFImporter/Source/GLTFCore/Private/GLTF/ExtensionsHandler.cpp
--- a/Engine/Plugins/Editor/GLTFImporter/Source/GLTFCore/Private/GLTF/ExtensionsHandler.cpp
+++ b/Engine/Plugins/Editor/GLTFImporter/Source/GLTFCore/Private/GLTF/ExtensionsHandler.cpp
@@ -11,8 +11,8 @@ namespace GLTF
 	namespace
 	{
 		static const TArray<FString> LightExtensions = {TEXT("KHR_lights_punctual"), TEXT("KHR_lights")};
-		static const TArray<FString> BaseExtensions  = LightExtensions;
-
+		static const TArray<FString> MaterialsExtensions = {TEXT("KHR_materials_variants")};
+	
 		TSharedPtr<FJsonObject> GetLightExtension(const TSharedPtr<FJsonObject>& Object)
 		{
 			if (!Object)
@@ -51,6 +51,11 @@ namespace GLTF
 	{
 		const TSharedPtr<FJsonObject>& ExtensionsObj = GetExtensions(Object);
 
+		if (!ExtensionsObj.IsValid())
+		{
+			return;
+		}
+
 		// lights
 
 		if (TSharedPtr<FJsonObject> LightsObj = GetLightExtension(ExtensionsObj))
@@ -68,7 +73,31 @@ namespace GLTF
 			}
 		}
 
-		CheckExtensions(Object, BaseExtensions);
+		// variants
+
+		if (ExtensionsObj->HasTypedField<EJson::Object>(TEXT("KHR_materials_variants")))
+		{
+			Asset->ExtensionsUsed.Add(EExtension::KHR_MaterialsVariants);
+
+			TSharedPtr<FJsonObject> VariantsObj = ExtensionsObj->GetObjectField(TEXT("KHR_materials_variants"));
+			uint32 VariantsCount = ArraySize(*VariantsObj, TEXT("variants"));
+			if (VariantsCount > 0)
+			{
+				Asset->Variants.Reserve(VariantsCount);
+				for (const TSharedPtr<FJsonValue>& Value : VariantsObj->GetArrayField(TEXT("variants")))
+				{
+					const TSharedPtr<FJsonObject>& NameObj = Value->AsObject();
+					const FString Name = NameObj->GetStringField(TEXT("name"));
+					Asset->Variants.Add(Name);
+				}
+			}
+		}
+		
+		TArray<FString> SupportedExtensions;
+		SupportedExtensions.Append(LightExtensions);
+		SupportedExtensions.Append(MaterialsExtensions);
+
+		CheckExtensions(Object, SupportedExtensions);
 	}
 
 	void FExtensionsHandler::SetupMaterialExtensions(const FJsonObject& Object, FMaterial& Material) const
@@ -223,7 +252,56 @@ namespace GLTF
 
 	void FExtensionsHandler::SetupPrimitiveExtensions(const FJsonObject& Object, FPrimitive& Primitive) const
 	{
-		static const TArray<FString> Extensions;
+		if (!Object.HasTypedField<EJson::Object>(TEXT("extensions")))
+		{
+			return;
+		}
+
+		enum
+		{
+			KHR_materials_variants = 0
+		};
+		static const TArray<FString> Extensions = { TEXT("KHR_materials_variants") };
+
+		const FJsonObject& ExtensionsObj = *Object.GetObjectField(TEXT("extensions"));
+		for (int32 Index = 0; Index < Extensions.Num(); ++Index)
+		{
+			const FString ExtensionName = Extensions[Index];
+			if (!ExtensionsObj.HasTypedField<EJson::Object>(ExtensionName))
+				continue;
+
+			const FJsonObject& ExtObj = *ExtensionsObj.GetObjectField(ExtensionName);
+			switch (Index)
+			{
+				case KHR_materials_variants:
+				{
+					const TArray<TSharedPtr<FJsonValue>>& Mappings = ExtObj.GetArrayField(TEXT("mappings"));
+
+					for (const TSharedPtr<FJsonValue>& Mapping : Mappings)
+					{
+						FVariantMapping& VariantMapping = Primitive.VariantMappings.Emplace_GetRef();
+
+						const TSharedPtr<FJsonObject> MappingObj = Mapping->AsObject();
+						VariantMapping.MaterialIndex = MappingObj->GetIntegerField(TEXT("material"));
+						const TArray<TSharedPtr<FJsonValue>>& Variants = MappingObj->GetArrayField(TEXT("variants"));
+
+						VariantMapping.VariantIndices.Reserve(Variants.Num());
+						for (const TSharedPtr<FJsonValue>& Variant : Variants)
+						{
+							const int32 VariantIndex = static_cast<int32>(Variant->AsNumber());
+							VariantMapping.VariantIndices.Add(VariantIndex);
+						}
+					}
+
+					Asset->ExtensionsUsed.Add(EExtension::KHR_MaterialsVariants);
+				}
+				break;
+				default:
+					check(false);
+					break;
+			}
+		}
+
 		CheckExtensions(Object, Extensions);
 	}
 
diff --git a/Engine/Plugins/Editor/GLTFImporter/Source/GLTFCore/Public/GLTFAsset.h b/Engine/Plugins/Editor/GLTFImporter/Source/GLTFCore/Public/GLTFAsset.h
--- a/Engine/Plugins/Editor/GLTFImporter/Source/GLTFCore/Public/GLTFAsset.h
+++ b/Engine/Plugins/Editor/GLTFImporter/Source/GLTFCore/Public/GLTFAsset.h
@@ -20,6 +20,7 @@ namespace GLTF
 		KHR_MaterialsClearCoat,
 		KHR_MaterialsTransmission,
 		KHR_MaterialsSheen,
+		KHR_MaterialsVariants,
 		KHR_TextureTransform,
 		KHR_DracoMeshCompression,
 		KHR_LightsPunctual,
@@ -78,6 +79,8 @@ namespace GLTF
 		TArray<FTexture>  Textures;
 		TArray<FMaterial> Materials;
 
+		TArray<FString>   Variants;
+
 		TSet<EExtension> ExtensionsUsed;
 		FMetadata        Metadata;
 
diff --git a/Engine/Plugins/Editor/GLTFImporter/Source/GLTFCore/Public/GLTFMesh.h b/Engine/Plugins/Editor/GLTFImporter/Source/GLTFCore/Public/GLTFMesh.h
--- a/Engine/Plugins/Editor/GLTFImporter/Source/GLTFCore/Public/GLTFMesh.h
+++ b/Engine/Plugins/Editor/GLTFImporter/Source/GLTFCore/Public/GLTFMesh.h
@@ -35,6 +35,12 @@ namespace GLTF
 		}
 	};
 
+	struct GLTFCORE_API FVariantMapping
+	{
+		int32 MaterialIndex;
+		TArray<int32> VariantIndices;
+	};
+
 	struct GLTFCORE_API FPrimitive
 	{
 		enum class EMode
@@ -53,6 +59,7 @@ namespace GLTF
 
 		const EMode Mode;
 		const int32 MaterialIndex;
+		TArray<FVariantMapping>	VariantMappings;
 
 		FPrimitive(EMode InMode, int32 InMaterial, const FAccessor& InIndices, const FAccessor& InPosition, const FAccessor& InNormal,
 		           const FAccessor& InTangent, const FAccessor& InTexCoord0, const FAccessor& InTexCoord1, const FAccessor& InColor0,
@@ -95,11 +102,12 @@ namespace GLTF
 		const FAccessor& Weights0;
 	};
 
+
 	struct GLTFCORE_API FMesh
 	{
-		FString            Name;
-		TArray<FPrimitive> Primitives;
-
+		FString				Name;
+		TArray<FPrimitive>	Primitives;
+	
 		bool HasNormals() const;
 		bool HasTangents() const;
 		bool HasTexCoords(uint32 Index) const;
diff --git a/Engine/Plugins/Editor/GLTFImporter/Source/GLTFImporter/Private/GLTFImporterContext.cpp b/Engine/Plugins/Editor/GLTFImporter/Source/GLTFImporter/Private/GLTFImporterContext.cpp
--- a/Engine/Plugins/Editor/GLTFImporter/Source/GLTFImporter/Private/GLTFImporterContext.cpp
+++ b/Engine/Plugins/Editor/GLTFImporter/Source/GLTFImporter/Private/GLTFImporterContext.cpp
@@ -46,7 +46,8 @@ bool FGLTFImporterContext::OpenFile(const FString& FilePath)
 	                                                             GLTF::EExtension::KHR_MaterialsUnlit,
 																 GLTF::EExtension::KHR_MaterialsClearCoat,
 																 GLTF::EExtension::KHR_MaterialsTransmission,
-																 GLTF::EExtension::KHR_MaterialsSheen};
+																 GLTF::EExtension::KHR_MaterialsSheen,
+																 GLTF::EExtension::KHR_MaterialsVariants};
 
 	for (GLTF::EExtension Extension : Asset.ExtensionsUsed)
 	{
diff --git a/Engine/Plugins/Enterprise/DatasmithGLTFImporter/Source/DatasmithGLTFTranslator/Private/DatasmithGLTFImporter.cpp b/Engine/Plugins/Enterprise/DatasmithGLTFImporter/Source/DatasmithGLTFTranslator/Private/DatasmithGLTFImporter.cpp
--- a/Engine/Plugins/Enterprise/DatasmithGLTFImporter/Source/DatasmithGLTFTranslator/Private/DatasmithGLTFImporter.cpp
+++ b/Engine/Plugins/Enterprise/DatasmithGLTFImporter/Source/DatasmithGLTFTranslator/Private/DatasmithGLTFImporter.cpp
@@ -102,7 +102,8 @@ bool FDatasmithGLTFImporter::OpenFile(const FString& InFileName)
 																 GLTF::EExtension::KHR_LightsPunctual,
 																 GLTF::EExtension::KHR_MaterialsClearCoat,
 																 GLTF::EExtension::KHR_MaterialsTransmission,
-																 GLTF::EExtension::KHR_MaterialsSheen};
+																 GLTF::EExtension::KHR_MaterialsSheen,
+																 GLTF::EExtension::KHR_MaterialsVariants};
 
 	for (GLTF::EExtension Extension : GLTFAsset->ExtensionsUsed)
 	{
@@ -244,9 +245,61 @@ TSharedPtr<IDatasmithMeshActorElement> FDatasmithGLTFImporter::CreateStaticMeshA
 
 	MeshActorElement = FDatasmithSceneFactory::CreateMeshActor(TEXT("TempName"));
 	MeshActorElement->SetStaticMeshPathName(*FDatasmithUtils::SanitizeObjectName(GLTFAsset->Meshes[MeshIndex].Name));
+
 	return MeshActorElement;
 }
 
+void FDatasmithGLTFImporter::CreateMaterialVariants(TSharedPtr<IDatasmithMeshActorElement> MeshActorElement, int32 MeshIndex)
+{
+	for (const GLTF::FPrimitive& Primitive : GLTFAsset->Meshes[MeshIndex].Primitives)
+	{
+		for (const GLTF::FVariantMapping& VariantMapping : Primitive.VariantMappings)
+		{
+			for (int32 VariantIndex : VariantMapping.VariantIndices)
+			{
+				if (!VariantSets.IsValid())
+				{
+					VariantSets = FDatasmithSceneFactory::CreateLevelVariantSets(TEXT("LevelVariantSets"));
+					VariantSet = FDatasmithSceneFactory::CreateVariantSet(TEXT("VariantSet"));
+					VariantSets->AddVariantSet(VariantSet.ToSharedRef());
+				}
+
+				ensure(GLTFAsset->Variants.IsValidIndex(VariantIndex));
+				const FString& VariantName = GLTFAsset->Variants[VariantIndex];
+
+				TSharedPtr<IDatasmithVariantElement> Variant;
+
+				if (!VariantNameToVariantElement.Contains(VariantName))
+				{
+					Variant = FDatasmithSceneFactory::CreateVariant(*VariantName);
+					VariantNameToVariantElement.Add(VariantName, Variant.ToSharedRef());
+					VariantSet->AddVariant(Variant.ToSharedRef());
+				}
+				else
+				{
+					Variant = VariantNameToVariantElement[VariantName];
+				}
+
+				const TArray<GLTF::FMaterialElement*>& Materials = MaterialFactory->GetMaterials();
+
+				FDatasmithGLTFMaterialElement* Material = static_cast<FDatasmithGLTFMaterialElement*>(Materials[VariantMapping.MaterialIndex]);
+				check(Material);
+
+				TSharedRef<IDatasmithObjectPropertyCaptureElement> PropCapture = FDatasmithSceneFactory::CreateObjectPropertyCapture();
+				PropCapture->SetCategory(EDatasmithPropertyCategory::Material);
+				PropCapture->SetRecordedObject(Material->GetMaterial());
+
+				// Assign mesh material to variant
+				TSharedRef<IDatasmithActorBindingElement> Binding = FDatasmithSceneFactory::CreateActorBinding();
+				Binding->SetActor(MeshActorElement);
+				Binding->AddPropertyCapture(PropCapture);
+			
+				Variant->AddActorBinding(Binding);
+			}
+		}
+	}
+}
+
 TSharedPtr<IDatasmithActorElement> FDatasmithGLTFImporter::ConvertNode(int32 NodeIndex)
 {
 	TSharedPtr<IDatasmithActorElement> ActorElement;
@@ -266,6 +319,8 @@ TSharedPtr<IDatasmithActorElement> FDatasmithGLTFImporter::ConvertNode(int32 Nod
 
 				ActorElement = MeshActorElement;
 				ActorElement->SetName(*Node.Name);
+
+				CreateMaterialVariants(MeshActorElement, Node.MeshIndex);
 			}
 			break;
 		case GLTF::FNode::EType::Camera:
@@ -343,6 +398,11 @@ bool FDatasmithGLTFImporter::SendSceneToDatasmith()
 	AnimationImporter->CurrentScene = &DatasmithScene.Get();
 	AnimationImporter->CreateAnimations(*GLTFAsset);
 
+	if (VariantSets.IsValid())
+	{
+		DatasmithScene->AddLevelVariantSets(VariantSets);
+	}
+
 	return true;
 }
 
diff --git a/Engine/Plugins/Enterprise/DatasmithGLTFImporter/Source/DatasmithGLTFTranslator/Private/DatasmithGLTFImporter.h b/Engine/Plugins/Enterprise/DatasmithGLTFImporter/Source/DatasmithGLTFTranslator/Private/DatasmithGLTFImporter.h
--- a/Engine/Plugins/Enterprise/DatasmithGLTFImporter/Source/DatasmithGLTFTranslator/Private/DatasmithGLTFImporter.h
+++ b/Engine/Plugins/Enterprise/DatasmithGLTFImporter/Source/DatasmithGLTFTranslator/Private/DatasmithGLTFImporter.h
@@ -18,6 +18,9 @@ class IDatasmithActorElement;
 class IDatasmithMeshActorElement;
 class IDatasmithBaseMaterialElement;
 class IDatasmithMeshElement;
+class IDatasmithVariantElement;
+class IDatasmithLevelVariantSetsElement;
+class IDatasmithVariantSetElement;
 class UDatasmithGLTFImportOptions;
 class FDatasmithGLTFAnimationImporter;
 class IDatasmithLevelSequenceElement;
@@ -78,6 +81,7 @@ private:
 	TSharedPtr<IDatasmithActorElement> CreateCameraActor(int32 CameraIndex) const;
 	TSharedPtr<IDatasmithActorElement> CreateLightActor(int32 CameraIndex) const;
 	TSharedPtr<IDatasmithMeshActorElement> CreateStaticMeshActor(int32 MeshIndex);
+	void CreateMaterialVariants(TSharedPtr<IDatasmithMeshActorElement> MeshActorElement, int32 MeshIndex);
 
 	TSharedPtr<IDatasmithActorElement> ConvertNode(int32 NodeIndex);
 
@@ -101,6 +105,9 @@ private:
 
 	TMap<IDatasmithMeshElement*, int32> MeshElementToGLTFMeshIndex; // track where Datasmith element originated from to extract data later
 	TMap<int32, TSharedRef<IDatasmithMeshElement>> GLTFMeshIndexToMeshElement;
+	TMap<FString, TSharedRef<IDatasmithVariantElement>> VariantNameToVariantElement;
+	TSharedPtr<IDatasmithLevelVariantSetsElement> VariantSets;
+	TSharedPtr<IDatasmithVariantSetElement> VariantSet;
 
 	const bool bTransformIsLocal = true; // GLTF nodes define transform in local space [REDACTED_PROFILE_URL]/glTF/tree/master/specification/2.0#transformations
 };
\ No newline at end of file
