DatasmithImport - Fix editor crash when instantiating a DatasmithScene in the level.

diff --git a/Engine/Plugins/Enterprise/DatasmithC4DImporter/Source/DatasmithC4DTranslator/Private/DatasmithC4DTranslator.cpp b/Engine/Plugins/Enterprise/DatasmithC4DImporter/Source/DatasmithC4DTranslator/Private/DatasmithC4DTranslator.cpp
--- a/Engine/Plugins/Enterprise/DatasmithC4DImporter/Source/DatasmithC4DTranslator/Private/DatasmithC4DTranslator.cpp
+++ b/Engine/Plugins/Enterprise/DatasmithC4DImporter/Source/DatasmithC4DTranslator/Private/DatasmithC4DTranslator.cpp
@@ -60,7 +60,10 @@ bool FDatasmithC4DTranslator::LoadScene(TSharedRef<IDatasmithScene> OutScene)
 
 void FDatasmithC4DTranslator::UnloadScene()
 {
-	Importer->UnloadScene();
+	if (Importer)
+	{
+		Importer->UnloadScene();
+	}
 }
 
 bool FDatasmithC4DTranslator::LoadStaticMesh(const TSharedRef<IDatasmithMeshElement> MeshElement, FDatasmithMeshElementPayload& OutMeshPayload)
diff --git a/Engine/Plugins/Enterprise/DatasmithCADImporter/Source/DatasmithPLMXMLTranslator/Private/DatasmithPlmXmlTranslator.cpp b/Engine/Plugins/Enterprise/DatasmithCADImporter/Source/DatasmithPLMXMLTranslator/Private/DatasmithPlmXmlTranslator.cpp
--- a/Engine/Plugins/Enterprise/DatasmithCADImporter/Source/DatasmithPLMXMLTranslator/Private/DatasmithPlmXmlTranslator.cpp
+++ b/Engine/Plugins/Enterprise/DatasmithCADImporter/Source/DatasmithPLMXMLTranslator/Private/DatasmithPlmXmlTranslator.cpp
@@ -71,8 +71,11 @@ bool FDatasmithPlmXmlTranslator::LoadScene(TSharedRef<IDatasmithScene> OutScene)
 
 void FDatasmithPlmXmlTranslator::UnloadScene()
 {
-	Importer->UnloadScene();
-	Importer.Reset();
+	if (Importer)
+	{
+		Importer->UnloadScene();
+		Importer.Reset();
+	}
 }
 
 bool FDatasmithPlmXmlTranslator::LoadStaticMesh(const TSharedRef<IDatasmithMeshElement> MeshElement, FDatasmithMeshElementPayload& OutMeshPayload)
diff --git a/Engine/Plugins/Enterprise/DatasmithFBXImporter/Source/DatasmithDeltaGenTranslator/Private/DatasmithDeltaGenTranslator.cpp b/Engine/Plugins/Enterprise/DatasmithFBXImporter/Source/DatasmithDeltaGenTranslator/Private/DatasmithDeltaGenTranslator.cpp
--- a/Engine/Plugins/Enterprise/DatasmithFBXImporter/Source/DatasmithDeltaGenTranslator/Private/DatasmithDeltaGenTranslator.cpp
+++ b/Engine/Plugins/Enterprise/DatasmithFBXImporter/Source/DatasmithDeltaGenTranslator/Private/DatasmithDeltaGenTranslator.cpp
@@ -101,7 +101,10 @@ bool FDatasmithDeltaGenTranslator::LoadScene(TSharedRef<IDatasmithScene> OutScen
 
 void FDatasmithDeltaGenTranslator::UnloadScene()
 {
-	Importer->UnloadScene();
+	if (Importer)
+	{
+		Importer->UnloadScene();
+	}
 }
 
 bool FDatasmithDeltaGenTranslator::LoadStaticMesh(const TSharedRef<IDatasmithMeshElement> MeshElement, FDatasmithMeshElementPayload& OutMeshPayload)
diff --git a/Engine/Plugins/Enterprise/DatasmithFBXImporter/Source/DatasmithVREDTranslator/Private/DatasmithVREDTranslator.cpp b/Engine/Plugins/Enterprise/DatasmithFBXImporter/Source/DatasmithVREDTranslator/Private/DatasmithVREDTranslator.cpp
--- a/Engine/Plugins/Enterprise/DatasmithFBXImporter/Source/DatasmithVREDTranslator/Private/DatasmithVREDTranslator.cpp
+++ b/Engine/Plugins/Enterprise/DatasmithFBXImporter/Source/DatasmithVREDTranslator/Private/DatasmithVREDTranslator.cpp
@@ -121,7 +121,10 @@ bool FDatasmithVREDTranslator::LoadScene(TSharedRef<IDatasmithScene> OutScene)
 
 void FDatasmithVREDTranslator::UnloadScene()
 {
-	Importer->UnloadScene();
+	if (Importer)
+	{
+		Importer->UnloadScene();
+	}
 }
 
 bool FDatasmithVREDTranslator::LoadStaticMesh(const TSharedRef<IDatasmithMeshElement> MeshElement, FDatasmithMeshElementPayload& OutMeshPayload)
diff --git a/Engine/Plugins/Enterprise/DatasmithGLTFImporter/Source/DatasmithGLTFTranslator/Private/DatasmithGLTFTranslator.cpp b/Engine/Plugins/Enterprise/DatasmithGLTFImporter/Source/DatasmithGLTFTranslator/Private/DatasmithGLTFTranslator.cpp
--- a/Engine/Plugins/Enterprise/DatasmithGLTFImporter/Source/DatasmithGLTFTranslator/Private/DatasmithGLTFTranslator.cpp
+++ b/Engine/Plugins/Enterprise/DatasmithGLTFImporter/Source/DatasmithGLTFTranslator/Private/DatasmithGLTFTranslator.cpp
@@ -146,8 +146,11 @@ bool FDatasmithGLTFTranslator::LoadScene(TSharedRef<IDatasmithScene> OutScene)
 
 void FDatasmithGLTFTranslator::UnloadScene()
 {
-	ShowLogMessages(Importer->GetLogMessages());
-	Importer->UnloadScene();
+	if (Importer)
+	{
+		ShowLogMessages(Importer->GetLogMessages());
+		Importer->UnloadScene();
+	}
 }
 
 bool FDatasmithGLTFTranslator::LoadStaticMesh(const TSharedRef<IDatasmithMeshElement> MeshElement, FDatasmithMeshElementPayload& OutMeshPayload)
diff --git a/Engine/Plugins/Enterprise/DatasmithIFCImporter/Source/DatasmithIFCTranslator/Private/DatasmithIFCTranslator.cpp b/Engine/Plugins/Enterprise/DatasmithIFCImporter/Source/DatasmithIFCTranslator/Private/DatasmithIFCTranslator.cpp
--- a/Engine/Plugins/Enterprise/DatasmithIFCImporter/Source/DatasmithIFCTranslator/Private/DatasmithIFCTranslator.cpp
+++ b/Engine/Plugins/Enterprise/DatasmithIFCImporter/Source/DatasmithIFCTranslator/Private/DatasmithIFCTranslator.cpp
@@ -89,7 +89,10 @@ bool FDatasmithIFCTranslator::LoadScene(TSharedRef<IDatasmithScene> OutScene)
 
 void FDatasmithIFCTranslator::UnloadScene()
 {
-	Importer->UnloadScene();
+	if (Importer)
+	{
+		Importer->UnloadScene();
+	}
 }
 
 bool FDatasmithIFCTranslator::LoadStaticMesh(const TSharedRef<IDatasmithMeshElement> MeshElement, FDatasmithMeshElementPayload& OutMeshPayload)
