Deprecate the Datasmith glTF translator and GLTFImporter

diff --git a/Engine/Plugins/Editor/GLTFImporter/GLTFImporter.uplugin b/Engine/Plugins/Editor/GLTFImporter/GLTFImporter.uplugin
--- a/Engine/Plugins/Editor/GLTFImporter/GLTFImporter.uplugin
+++ b/Engine/Plugins/Editor/GLTFImporter/GLTFImporter.uplugin
@@ -2,8 +2,8 @@
 	"FileVersion": 3,
 	"Version": 1.0,
 	"VersionName": "1.0",
-	"FriendlyName": "glTF Importer",
-	"Description": "An importer for Khronos glTF 3D models.",
+	"FriendlyName": "glTF Importer (Deprecated)",
+	"Description": "This plugin has been deprecated and will be removed in a future engine version.\nPlease update your project to use glTF import through the Interchange plugin instead.\nAn importer for Khronos glTF 3D models.",
 	"Category": "Importers",
 	"CreatedBy": "Epic Games, Inc.",
 	"CreatedByURL": "https://epicgames.com",
diff --git a/Engine/Plugins/Enterprise/DatasmithGLTFImporter/DatasmithGLTFImporter.uplugin b/Engine/Plugins/Enterprise/DatasmithGLTFImporter/DatasmithGLTFImporter.uplugin
--- a/Engine/Plugins/Enterprise/DatasmithGLTFImporter/DatasmithGLTFImporter.uplugin
+++ b/Engine/Plugins/Enterprise/DatasmithGLTFImporter/DatasmithGLTFImporter.uplugin
@@ -2,8 +2,8 @@
 	"FileVersion": 3,
 	"Version": 1.0,
 	"VersionName": "1.0",
-	"FriendlyName": "Datasmith glTF Importer",
-	"Description": "An importer for Khronos glTF 3D models.",
+	"FriendlyName": "Datasmith glTF Importer (Deprecated)",
+	"Description": "This plugin has been deprecated and will be removed in a future engine version.\nPlease update your project to use glTF import through the Interchange plugin instead.\nAn importer for Khronos glTF 3D models.",
 	"Category": "Importers",
 	"CreatedBy": "Epic Games, Inc.",
 	"CreatedByURL": "https://epicgames.com",
