[GLTF Import] Removing unnecessary GenerateNames call.
In CL 23958187 the GenerateNames function have been added to the GLTF::FFileReader::ReadFile function, hence not needed to call it separately anymore.

diff --git a/Engine/Plugins/Editor/GLTFImporter/Source/GLTFImporter/Private/GLTFImporterContext.cpp b/Engine/Plugins/Editor/GLTFImporter/Source/GLTFImporter/Private/GLTFImporterContext.cpp
--- a/Engine/Plugins/Editor/GLTFImporter/Source/GLTFImporter/Private/GLTFImporterContext.cpp
+++ b/Engine/Plugins/Editor/GLTFImporter/Source/GLTFImporter/Private/GLTFImporterContext.cpp
@@ -55,8 +55,6 @@ bool FGLTFImporterContext::OpenFile(const FString& FilePath)
 		}
 	}
 
-	Asset.GenerateNames(FPaths::GetBaseFilename(FilePath));
-
 	return true;
 }
 
diff --git a/Engine/Plugins/Enterprise/DatasmithGLTFImporter/Source/DatasmithGLTFTranslator/Private/DatasmithGLTFImporter.cpp b/Engine/Plugins/Enterprise/DatasmithGLTFImporter/Source/DatasmithGLTFTranslator/Private/DatasmithGLTFImporter.cpp
--- a/Engine/Plugins/Enterprise/DatasmithGLTFImporter/Source/DatasmithGLTFTranslator/Private/DatasmithGLTFImporter.cpp
+++ b/Engine/Plugins/Enterprise/DatasmithGLTFImporter/Source/DatasmithGLTFTranslator/Private/DatasmithGLTFImporter.cpp
@@ -95,8 +95,6 @@ bool FDatasmithGLTFImporter::OpenFile(const FString& InFileName)
 	}
 	check(GLTFAsset->ValidationCheck() == GLTF::FAsset::Valid);
 
-	GLTFAsset->GenerateNames(FPaths::GetBaseFilename(InFileName));
-
 	// check extensions supported
 	for (GLTF::EExtension Extension : GLTFAsset->ExtensionsUsed)
 	{
